Fixed some type conversion warnings with new clang/xcode

diff --git a/Engine/Plugins/Runtime/AR/AzureSpatialAnchorsForARKit/Source/AzureSpatialAnchorsForARKit/Private/AzureSpatialAnchorsARKitInterop.cpp b/Engine/Plugins/Runtime/AR/AzureSpatialAnchorsForARKit/Source/AzureSpatialAnchorsForARKit/Private/AzureSpatialAnchorsARKitInterop.cpp
--- a/Engine/Plugins/Runtime/AR/AzureSpatialAnchorsForARKit/Source/AzureSpatialAnchorsForARKit/Private/AzureSpatialAnchorsARKitInterop.cpp
+++ b/Engine/Plugins/Runtime/AR/AzureSpatialAnchorsForARKit/Source/AzureSpatialAnchorsForARKit/Private/AzureSpatialAnchorsARKitInterop.cpp
@@ -163,12 +163,12 @@ void FAzureSpatialAnchorsARKitInterop::GetAccessTokenWithAccountKeyAsync(const F
         std::wstringstream string;
         if (error != nullptr)
         {
-            string << L"Failed to obtain access token from account key: " << *NSStringToFString(error.localizedDescription);
+            string << L"Failed to obtain access token from account key: " << TCHAR_TO_WCHAR(*NSStringToFString(error.localizedDescription));
             Callback(EAzureSpatialAnchorsResult::FailSeeErrorString, string.str().c_str(), L"");
             return;
         }
         
-        string << *NSStringToFString(value);
+        string << TCHAR_TO_WCHAR(*NSStringToFString(value));
         Callback(EAzureSpatialAnchorsResult::Success, L"", string.str().c_str());
     }];
 }
@@ -185,12 +185,12 @@ void FAzureSpatialAnchorsARKitInterop::GetAccessTokenWithAuthenticationTokenAsyn
         std::wstringstream string;
         if (error != nullptr)
         {
-            string << L"Failed to obtain access token from authentication token: " << *NSStringToFString(error.localizedDescription);
+            string << L"Failed to obtain access token from authentication token: " << TCHAR_TO_WCHAR(*NSStringToFString(error.localizedDescription));
             Callback(EAzureSpatialAnchorsResult::FailSeeErrorString, string.str().c_str(), L"");
             return;
         }
         
-        string << *NSStringToFString(value);
+        string << TCHAR_TO_WCHAR(*NSStringToFString(value));
         Callback(EAzureSpatialAnchorsResult::Success, L"", string.str().c_str());
     }];
 }
@@ -319,7 +319,7 @@ void FAzureSpatialAnchorsARKitInterop::GetSessionStatusAsync(IAzureSpatialAnchor
         if (error != nullptr)
         {
             std::wstringstream string;
-            string << L"Failed to obtain session status: " << *NSStringToFString(error.localizedDescription);
+            string << L"Failed to obtain session status: " << TCHAR_TO_WCHAR(*NSStringToFString(error.localizedDescription));
             Callback(EAzureSpatialAnchorsResult::FailSeeErrorString, string.str().c_str(), SessionStatus);
             return;
         }
@@ -426,7 +426,7 @@ void FAzureSpatialAnchorsARKitInterop::CreateAnchorAsync(IAzureSpatialAnchors::C
         if (error != nullptr)
         {
             std::wstringstream string;
-            string << L"Failed to create azure spatial anchor: " << *NSStringToFString(error.localizedDescription) << L", " << *NSStringToFString(error.localizedFailureReason);
+            string << L"Failed to create azure spatial anchor: " << TCHAR_TO_WCHAR(*NSStringToFString(error.localizedDescription)) << L", " << TCHAR_TO_WCHAR(*NSStringToFString(error.localizedFailureReason));
             Callback(EAzureSpatialAnchorsResult::FailSeeErrorString, string.str().c_str());
             return;
         }
@@ -466,7 +466,7 @@ void FAzureSpatialAnchorsARKitInterop::DeleteAnchorAsync(IAzureSpatialAnchors::C
         if (error != nullptr)
         {
            std::wstringstream string;
-           string << L"Failed to delete azure spatial anchor: " << *NSStringToFString(error.localizedDescription);
+           string << L"Failed to delete azure spatial anchor: " << TCHAR_TO_WCHAR(*NSStringToFString(error.localizedDescription));
            Callback(EAzureSpatialAnchorsResult::FailSeeErrorString, string.str().c_str());
            return;
         }
@@ -590,7 +590,7 @@ void FAzureSpatialAnchorsARKitInterop::GetAnchorPropertiesAsync(const FString& I
         if (error != nullptr)
         {
             std::wstringstream string;
-            string << L"Failed to get azure spatial anchor properties: " << *NSStringToFString(error.localizedDescription);
+            string << L"Failed to get azure spatial anchor properties: " << TCHAR_TO_WCHAR(*NSStringToFString(error.localizedDescription));
             Callback(EAzureSpatialAnchorsResult::FailSeeErrorString, string.str().c_str(), IAzureSpatialAnchors::CloudAnchorID_Invalid);
             return;
         }
@@ -630,7 +630,7 @@ void FAzureSpatialAnchorsARKitInterop::RefreshAnchorPropertiesAsync(IAzureSpatia
         if (error != nullptr)
         {
            std::wstringstream string;
-           string << L"Failed to refresh azure spatial anchor properties: " <<  *NSStringToFString(error.localizedDescription);
+           string << L"Failed to refresh azure spatial anchor properties: " <<  TCHAR_TO_WCHAR(*NSStringToFString(error.localizedDescription));
            Callback(EAzureSpatialAnchorsResult::FailSeeErrorString, string.str().c_str());
            return;
         }
@@ -664,7 +664,7 @@ void FAzureSpatialAnchorsARKitInterop::UpdateAnchorPropertiesAsync(IAzureSpatial
         if (error != nullptr)
         {
            std::wstringstream string;
-           string << L"Failed to update azure spatial anchor properties: " <<  *NSStringToFString(error.localizedDescription);
+           string << L"Failed to update azure spatial anchor properties: " << TCHAR_TO_WCHAR(*NSStringToFString(error.localizedDescription));
            Callback(EAzureSpatialAnchorsResult::FailSeeErrorString, string.str().c_str());
            return;
         }
@@ -1008,12 +1008,12 @@ void FAzureSpatialAnchorsARKitInterop::CreateDiagnosticsManifestAsync(const FStr
         std::wstringstream string;
         if (error != nullptr)
         {
-           string << L"Failed to create diagnostics manifest: " <<  *NSStringToFString(error.localizedDescription);
+           string << L"Failed to create diagnostics manifest: " << TCHAR_TO_WCHAR(*NSStringToFString(error.localizedDescription));
            Callback(EAzureSpatialAnchorsResult::FailSeeErrorString, string.str().c_str(), L"");
            return;
         }
         
-        string << *NSStringToFString(value);
+        string << TCHAR_TO_WCHAR(*NSStringToFString(value));
         Callback(EAzureSpatialAnchorsResult::Success, L"", string.str().c_str());
     }];
 }
@@ -1030,7 +1030,7 @@ void FAzureSpatialAnchorsARKitInterop::SubmitDiagnosticsManifestAsync(const FStr
         if (error != nullptr)
         {
            std::wstringstream string;
-           string << L"Failed to submit diagnostics manifest: " <<  *NSStringToFString(error.localizedDescription);
+           string << L"Failed to submit diagnostics manifest: " << TCHAR_TO_WCHAR(*NSStringToFString(error.localizedDescription));
            Callback(EAzureSpatialAnchorsResult::FailSeeErrorString, string.str().c_str());
            return;
         }
