UE-200861 AzureSpatialAnchors service retirement deprecation.
-Marking the plugins as deprecated and logging a warning.

diff --git a/Engine/Plugins/Runtime/AR/AzureSpatialAnchors/AzureSpatialAnchors.uplugin b/Engine/Plugins/Runtime/AR/AzureSpatialAnchors/AzureSpatialAnchors.uplugin
--- a/Engine/Plugins/Runtime/AR/AzureSpatialAnchors/AzureSpatialAnchors.uplugin
+++ b/Engine/Plugins/Runtime/AR/AzureSpatialAnchors/AzureSpatialAnchors.uplugin
@@ -1,25 +1,25 @@
 {
-	"FileVersion" : 3,
-	"Version" : 1,
-	"VersionName" : "1.0",
-	"FriendlyName" : "Azure Spatial Anchors",
-	"Description" : "Support for Microsoft Azure Spatial Anchor cloud service.",
-	"Category" : "Augmented Reality",
-	"CreatedBy" : "Epic Games, Inc.",
-	"CreatedByURL" : "https://epicgames.com",
-	"DocsURL" : "",
-	"MarketplaceURL" : "",
-	"SupportURL" : "",
-	"EnabledByDefault" : false,
-	"CanContainContent" : false,
-	"IsExperimentalVersion" : true,
-	"IsBetaVersion" : true,
-	"Installed" : false,
+	"FileVersion": 3,
+	"Version": 1,
+	"VersionName": "1.0",
+	"FriendlyName": "Azure Spatial Anchors (Deprecated, service ending)",
+	"Description": "Microsoft Azure Spatial Anchor cloud service will cease to function November 20, 2024.",
+	"Category": "Augmented Reality",
+	"CreatedBy": "Epic Games, Inc.",
+	"CreatedByURL": "https://epicgames.com",
+	"DocsURL": "",
+	"MarketplaceURL": "",
+	"SupportURL": "",
+	"EnabledByDefault": false,
+	"CanContainContent": false,
+	"IsExperimentalVersion": true,
+	"IsBetaVersion": true,
+	"Installed": false,
 	"SupportedTargetPlatforms": [
 		"Win64",
 		"Android",
 		"IOS",
-        	"Mac"
+		"Mac"
 	],
 	"Modules": [
 		{
diff --git a/Engine/Plugins/Runtime/AR/AzureSpatialAnchors/Source/AzureSpatialAnchors/Private/AzureSpatialAnchorsFunctionLibrary.cpp b/Engine/Plugins/Runtime/AR/AzureSpatialAnchors/Source/AzureSpatialAnchors/Private/AzureSpatialAnchorsFunctionLibrary.cpp
--- a/Engine/Plugins/Runtime/AR/AzureSpatialAnchors/Source/AzureSpatialAnchors/Private/AzureSpatialAnchorsFunctionLibrary.cpp
+++ b/Engine/Plugins/Runtime/AR/AzureSpatialAnchors/Source/AzureSpatialAnchors/Private/AzureSpatialAnchorsFunctionLibrary.cpp
@@ -16,6 +16,8 @@ bool UAzureSpatialAnchorsLibrary::CreateSession()
 		return false;
 	}
 
+	UE_LOG(LogAzureSpatialAnchors, Warning, TEXT("Azure Spatial Anchors (ASA) will be retired by microsoft on 20 November 2024."));
+
 	return IASA->CreateSession();
 }
 
diff --git a/Engine/Plugins/Runtime/AR/AzureSpatialAnchorsForARCore/AzureSpatialAnchorsForARCore.uplugin b/Engine/Plugins/Runtime/AR/AzureSpatialAnchorsForARCore/AzureSpatialAnchorsForARCore.uplugin
--- a/Engine/Plugins/Runtime/AR/AzureSpatialAnchorsForARCore/AzureSpatialAnchorsForARCore.uplugin
+++ b/Engine/Plugins/Runtime/AR/AzureSpatialAnchorsForARCore/AzureSpatialAnchorsForARCore.uplugin
@@ -1,20 +1,20 @@
 {
-	"FileVersion" : 3,
-	"Version" : 1,
-	"VersionName" : "1.0",
-	"FriendlyName" : "Azure Spatial Anchors For ARCore",
-	"Description" : "Microsoft Azure Spatial Anchors implementation for ARCore",
-	"Category" : "Augmented Reality",
-	"CreatedBy" : "Epic Games, Inc.",
-	"CreatedByURL" : "https://epicgames.com",
-	"DocsURL" : "",
-	"MarketplaceURL" : "",
-	"SupportURL" : "",
-	"EnabledByDefault" : false,
-	"CanContainContent" : false,
-	"IsExperimentalVersion" : true,
-	"IsBetaVersion" : true,
-	"Installed" : false,
+	"FileVersion": 3,
+	"Version": 1,
+	"VersionName": "1.0",
+	"FriendlyName": "Azure Spatial Anchors For ARCore (Deprecated, service ending)",
+	"Description": "Microsoft Azure Spatial Anchor cloud service will cease to function November 20, 2024.",
+	"Category": "Augmented Reality",
+	"CreatedBy": "Epic Games, Inc.",
+	"CreatedByURL": "https://epicgames.com",
+	"DocsURL": "",
+	"MarketplaceURL": "",
+	"SupportURL": "",
+	"EnabledByDefault": false,
+	"CanContainContent": false,
+	"IsExperimentalVersion": true,
+	"IsBetaVersion": true,
+	"Installed": false,
 	"SupportedTargetPlatforms": [
 		"Android"
 	],
diff --git a/Engine/Plugins/Runtime/AR/AzureSpatialAnchorsForARKit/AzureSpatialAnchorsForARKit.uplugin b/Engine/Plugins/Runtime/AR/AzureSpatialAnchorsForARKit/AzureSpatialAnchorsForARKit.uplugin
--- a/Engine/Plugins/Runtime/AR/AzureSpatialAnchorsForARKit/AzureSpatialAnchorsForARKit.uplugin
+++ b/Engine/Plugins/Runtime/AR/AzureSpatialAnchorsForARKit/AzureSpatialAnchorsForARKit.uplugin
@@ -1,20 +1,20 @@
 {
-	"FileVersion" : 3,
-	"Version" : 1,
-	"VersionName" : "1.0",
-	"FriendlyName" : "Azure Spatial Anchors For ARKit",
-	"Description" : "Microsoft Azure Spatial Anchors implementation for ARKit",
-	"Category" : "Augmented Reality",
-	"CreatedBy" : "Epic Games, Inc.",
-	"CreatedByURL" : "https://epicgames.com",
-	"DocsURL" : "",
-	"MarketplaceURL" : "",
-	"SupportURL" : "",
-	"EnabledByDefault" : false,
-	"CanContainContent" : false,
-	"IsExperimentalVersion" : true,
-	"IsBetaVersion" : true,
-	"Installed" : false,
+	"FileVersion": 3,
+	"Version": 1,
+	"VersionName": "1.0",
+	"FriendlyName": "Azure Spatial Anchors For ARKit (Deprecated, service ending)",
+	"Description": "Microsoft Azure Spatial Anchor cloud service will cease to function November 20, 2024.",
+	"Category": "Augmented Reality",
+	"CreatedBy": "Epic Games, Inc.",
+	"CreatedByURL": "https://epicgames.com",
+	"DocsURL": "",
+	"MarketplaceURL": "",
+	"SupportURL": "",
+	"EnabledByDefault": false,
+	"CanContainContent": false,
+	"IsExperimentalVersion": true,
+	"IsBetaVersion": true,
+	"Installed": false,
 	"SupportedTargetPlatforms": [
 		"IOS"
 	],
@@ -35,12 +35,12 @@
 				"IOS"
 			]
 		},
-        {
-            "Name": "AppleARKit",
-            "Enabled": true,
-            "PlatformAllowList": [
-                "IOS"
-            ]
-        }
+		{
+			"Name": "AppleARKit",
+			"Enabled": true,
+			"PlatformAllowList": [
+				"IOS"
+			]
+		}
 	]
 }
